Fix WidgetCarousel including headers from UMG.  Slate (not in engine) can never depend on UMG (in engine)

diff --git a/Engine/Source/Runtime/WidgetCarousel/Public/SWidgetCarousel.h b/Engine/Source/Runtime/WidgetCarousel/Public/SWidgetCarousel.h
--- a/Engine/Source/Runtime/WidgetCarousel/Public/SWidgetCarousel.h
+++ b/Engine/Source/Runtime/WidgetCarousel/Public/SWidgetCarousel.h
@@ -16,7 +16,6 @@
 #include "Widgets/SOverlay.h"
 #include "Animation/CurveSequence.h"
 #include "Widgets/Layout/SFxWidget.h"
-#include "Components/SlateWrapperTypes.h"
 
 /** The desired Carousel scroll direction */
 namespace EWidgetCarouselScrollDirection
