Fixing a crash in the widget carousel when you start an animation before things are ready, and a non-existant index is selected.

diff --git a/Engine/Source/Runtime/WidgetCarousel/Public/SWidgetCarousel.h b/Engine/Source/Runtime/WidgetCarousel/Public/SWidgetCarousel.h
--- a/Engine/Source/Runtime/WidgetCarousel/Public/SWidgetCarousel.h
+++ b/Engine/Source/Runtime/WidgetCarousel/Public/SWidgetCarousel.h
@@ -496,6 +496,11 @@ public:
 		const TArray<ItemType>& ItemsSourceRef = (*ItemsSource);
 		if (ItemsSourceRef.Num() > 0)
 		{
+			if (WidgetIndex >= ItemsSourceRef.Num())
+			{
+				WidgetIndex = WidgetIndex % ItemsSourceRef.Num();
+			}
+
 			LeftCarouselWidget->SetWidgetContent(OnGenerateWidget.Execute(ItemsSourceRef[GetLeftWidgetIndex(WidgetIndex)]));
 			CenterCarouselWidget->SetWidgetContent(OnGenerateWidget.Execute(ItemsSourceRef[WidgetIndex]));
 			RightCarouselWidget->SetWidgetContent(OnGenerateWidget.Execute(ItemsSourceRef[GetRightWidgetIndex(WidgetIndex)]));
@@ -706,7 +711,7 @@ protected:
 	int32 GetRightWidgetIndex(int32 Index)
 	{
 		int32 RightWidgetIndex = Index + 1;
-		if (RightWidgetIndex > ItemsSource->Num() - 1)
+		if (RightWidgetIndex >= ItemsSource->Num())
 		{
 			RightWidgetIndex = 0;
 		}
