CommonUI: Fix assert when using a CommonWidgetCarousel widget when 1 or 2 children with Global Invalidation enabled

- Slate GI doesn't allow widgets to exist in multiple slots
  - When the CommonWidgetCarousel had fewer then 3 children, it would reuse children in the left and right slots
- Reused widgets are now replaced with `SNullWidget::NullWidget` in slots that will not be visible during transitions
- Also fixed not being able to accurately preview the selected widget due to transitions being active when selecting widgets in the designer

#tests Added a CommonWidgetCarousel with 0, 1, 2, and 3 children to the screen and ran PIE with Slate GI enabled
[RN] fix creator/ui

diff --git a/Engine/Plugins/Runtime/CommonUI/Source/CommonUI/Private/CommonWidgetCarousel.cpp b/Engine/Plugins/Runtime/CommonUI/Source/CommonUI/Private/CommonWidgetCarousel.cpp
--- a/Engine/Plugins/Runtime/CommonUI/Source/CommonUI/Private/CommonWidgetCarousel.cpp
+++ b/Engine/Plugins/Runtime/CommonUI/Source/CommonUI/Private/CommonWidgetCarousel.cpp
@@ -207,8 +207,10 @@ void UCommonWidgetCarousel::OnDescendantSelectedByDesigner(UWidget* DescendantWi
 	if ( SelectedChild )
 	{
 		int32 OverrideIndex = GetChildIndex(SelectedChild);
-		if ( OverrideIndex != -1 && MyCommonWidgetCarousel.IsValid() )
+		if ( OverrideIndex != INDEX_NONE && MyCommonWidgetCarousel.IsValid() )
 		{
+			// Need to cancel transition if it's active to ensure the selected widget actually displays if there's a transition in progress
+			MyCommonWidgetCarousel->ResetTransition();
 			MyCommonWidgetCarousel->SetActiveWidgetIndex(OverrideIndex);
 		}
 	}
@@ -216,6 +218,12 @@ void UCommonWidgetCarousel::OnDescendantSelectedByDesigner(UWidget* DescendantWi
 
 void UCommonWidgetCarousel::OnDescendantDeselectedByDesigner(UWidget* DescendantWidget)
 {
+	if (MyCommonWidgetCarousel.IsValid())
+	{
+		// Need to cancel transition if it's active to ensure the active widget actually displays if there's a transition in progress
+		MyCommonWidgetCarousel->ResetTransition();
+	}
+
 	SetActiveWidgetIndex(ActiveWidgetIndex);
 }
 
diff --git a/Engine/Source/Runtime/WidgetCarousel/Public/SWidgetCarousel.h b/Engine/Source/Runtime/WidgetCarousel/Public/SWidgetCarousel.h
--- a/Engine/Source/Runtime/WidgetCarousel/Public/SWidgetCarousel.h
+++ b/Engine/Source/Runtime/WidgetCarousel/Public/SWidgetCarousel.h
@@ -210,6 +210,20 @@ struct FCarouselDisplayItem : public TSharedFromThis<FCarouselDisplayItem>
 		];
 	}
 
+	TSharedRef<SWidget> GetWidgetContent() const
+	{
+		if (FXWidget.IsValid())
+		{
+			FChildren* Children = FXWidget->GetChildren();
+			if (Children && Children->Num() > 0)
+			{
+				return Children->GetChildAt(0);
+			}		
+		}
+
+		return SNullWidget::NullWidget;
+	}
+
 //	FText GetPositionText() const
 //	{
 //		return FText::FromString(FString::Printf(TEXT("%.2f, %.2f, %.2f"), SlideValue, DesiredSlideValue, SlideInCurve.GetLerp()));
@@ -505,11 +519,18 @@ public:
 			{
 				WidgetIndex = WidgetIndex % ItemsSourceRef.Num();
 			}
-
-			LeftCarouselWidget->SetWidgetContent(OnGenerateWidget.Execute(ItemsSourceRef[GetLeftWidgetIndex(WidgetIndex)]));
-			RightCarouselWidget->SetWidgetContent(OnGenerateWidget.Execute(ItemsSourceRef[GetRightWidgetIndex(WidgetIndex)]));
+			
+			TSharedRef<SWidget> LeftWidget = OnGenerateWidget.Execute(ItemsSourceRef[GetLeftWidgetIndex(WidgetIndex)]);
+			TSharedRef<SWidget> RightWidget = OnGenerateWidget.Execute(ItemsSourceRef[GetRightWidgetIndex(WidgetIndex)]);
+			TSharedRef<SWidget> CenterWidget = OnGenerateWidget.Execute(ItemsSourceRef[WidgetIndex]);
+
+			// Each carousel widget must have unique content since Slate Global Invalidation requires that widgets only exist in a single slot
+			const bool bIsLeftWidgetUnique = (LeftWidget != CenterWidget) && (RightWidget != LeftWidget);
+			const bool bIsRightWidgetUnique = RightWidget != CenterWidget;
+			LeftCarouselWidget->SetWidgetContent(bIsLeftWidgetUnique ? LeftWidget : SNullWidget::NullWidget);
+			RightCarouselWidget->SetWidgetContent(bIsRightWidgetUnique ? RightWidget : SNullWidget::NullWidget);
 			// CenterCarouselWidget is the one that displays the active child.  Needs to be populated last for the Widget Path to match the visible hierarchy.
-			CenterCarouselWidget->SetWidgetContent(OnGenerateWidget.Execute(ItemsSourceRef[WidgetIndex]));
+			CenterCarouselWidget->SetWidgetContent(CenterWidget);
 		}
 
 		if (ItemsSourceRef.Num() == 1)
@@ -622,6 +643,14 @@ public:
 		return CenterCarouselWidget->GetMoveSpeed();
 	}
 
+	void ResetTransition()
+	{
+		constexpr bool bClearTransition = true;
+		CenterCarouselWidget->SetSlide(0.f, bClearTransition);
+		LeftCarouselWidget->SetSlide(SlideValueLeftLimit, bClearTransition);
+		RightCarouselWidget->SetSlide(SlideValueRightLimit, bClearTransition);
+	}
+
 protected:
 
 	/**
@@ -651,23 +680,31 @@ protected:
 	/**
 	 * Swap the buffer index.
 	 */
-	void SwapBuffer(EWidgetCarouselScrollDirection::Type ScrollDirection, int32 OverrideWidget = INDEX_NONE)
+	void SwapBuffer(EWidgetCarouselScrollDirection::Type ScrollDirection, int32 OverrideWidgetIndex = INDEX_NONE)
 	{
 		// Register for an active update every 2.5 seconds
 		RegisterActiveTimer(0.0f, FWidgetActiveTimerDelegate::CreateSP(this, &SWidgetCarousel::RegisteredTick));
 
 		const TArray<ItemType>& ItemsSourceRef = (*ItemsSource);
+		
+		TSharedRef<SWidget> LeftWidget = LeftCarouselWidget->GetWidgetContent();
+		TSharedRef<SWidget> RightWidget = RightCarouselWidget->GetWidgetContent();
+		TSharedRef<SWidget> CenterWidget = CenterCarouselWidget->GetWidgetContent();
+
+		const bool bIsScrollingLeft = ScrollDirection == EWidgetCarouselScrollDirection::Carousel_Left;
 
 		if (ItemsSourceRef.Num() > 0)
 		{
-			if (ScrollDirection == EWidgetCarouselScrollDirection::Carousel_Left)
+			if (bIsScrollingLeft)
 			{
 				WidgetIndex = GetLeftWidgetIndex(WidgetIndex);
 				TSharedPtr<FCarouselDisplayItem> TempWidget = RightCarouselWidget;
 				RightCarouselWidget = CenterCarouselWidget;
+				RightWidget = (CenterWidget == SNullWidget::NullWidget) ? OnGenerateWidget.Execute(ItemsSourceRef[GetRightWidgetIndex(WidgetIndex)]) : CenterWidget;
 				CenterCarouselWidget = LeftCarouselWidget;
+				CenterWidget = (LeftWidget == SNullWidget::NullWidget) ? OnGenerateWidget.Execute(ItemsSourceRef[WidgetIndex]) : LeftWidget;
 				LeftCarouselWidget = TempWidget;
-				LeftCarouselWidget->SetWidgetContent(OnGenerateWidget.Execute(ItemsSourceRef[GetLeftWidgetIndex(WidgetIndex)]));
+				LeftWidget = OnGenerateWidget.Execute(ItemsSourceRef[GetLeftWidgetIndex(WidgetIndex)]);
 				LeftCarouselWidget->SetSlide(CenterCarouselWidget->GetSlideValue() + SlideValueLeftLimit, true);
 			}
 			else
@@ -675,19 +712,28 @@ protected:
 				WidgetIndex = GetRightWidgetIndex(WidgetIndex);
 				TSharedPtr<FCarouselDisplayItem> TempWidget = LeftCarouselWidget;
 				LeftCarouselWidget = CenterCarouselWidget;
+				LeftWidget = (CenterWidget == SNullWidget::NullWidget) ? OnGenerateWidget.Execute(ItemsSourceRef[GetLeftWidgetIndex(WidgetIndex)]) : CenterWidget;
 				CenterCarouselWidget = RightCarouselWidget;
+				CenterWidget = (RightWidget == SNullWidget::NullWidget) ? OnGenerateWidget.Execute(ItemsSourceRef[WidgetIndex]) : RightWidget;
 				RightCarouselWidget = TempWidget;
-				RightCarouselWidget->SetWidgetContent(OnGenerateWidget.Execute(ItemsSourceRef[GetRightWidgetIndex(WidgetIndex)]));
+				RightWidget = OnGenerateWidget.Execute(ItemsSourceRef[GetRightWidgetIndex(WidgetIndex)]);
 				RightCarouselWidget->SetSlide(CenterCarouselWidget->GetSlideValue() + SlideValueRightLimit, true);
 			}
 		}
 
-		if (OverrideWidget != INDEX_NONE)
+		if (ItemsSourceRef.IsValidIndex(OverrideWidgetIndex))
 		{
 			// Need to find a way to clear left and right images
-			WidgetIndex = OverrideWidget;
-			CenterCarouselWidget->SetWidgetContent(OnGenerateWidget.Execute(ItemsSourceRef[OverrideWidget]));
+			WidgetIndex = OverrideWidgetIndex;
+			CenterWidget = OnGenerateWidget.Execute(ItemsSourceRef[OverrideWidgetIndex]);
 		}
+		
+		// Ensure the 3 CarouselWidgets here have unique widgets, if that isn't the case, prioritize the 2 widgets that will display during the transition
+		const bool bIsLeftWidgetUnique = (LeftWidget != CenterWidget) && (!bIsScrollingLeft || RightWidget != LeftWidget);
+		const bool bIsRightWidgetUnique = (RightWidget != CenterWidget) && (bIsScrollingLeft || RightWidget != LeftWidget);
+		LeftCarouselWidget->SetWidgetContent(bIsLeftWidgetUnique ? LeftWidget : SNullWidget::NullWidget);
+		RightCarouselWidget->SetWidgetContent(bIsRightWidgetUnique ? RightWidget : SNullWidget::NullWidget);
+		CenterCarouselWidget->SetWidgetContent(CenterWidget);
 
 		SetSliderLimits();
 		GenerateWidgets();
